#vr Experimental:  Adding in StereoPanorama plugin, which supports exporting stereoscopic movies from the engine

diff --git a/Engine/Plugins/Experimental/StereoPanorama/Source/StereoPanorama/Private/SceneCapturer.cpp b/Engine/Plugins/Experimental/StereoPanorama/Source/StereoPanorama/Private/SceneCapturer.cpp
new file mode 100644
--- /dev/null
+++ b/Engine/Plugins/Experimental/StereoPanorama/Source/StereoPanorama/Private/SceneCapturer.cpp
@@ -0,0 +1,857 @@
+// Copyright 2015 Kite & Lightning.  All rights reserved.
+
+#include "StereoPanoramaPrivatePCH.h"
+
+DEFINE_LOG_CATEGORY( LogStereoPanorama );
+
+//Rotated Grid Supersampling
+const int32 maxNumSamples = 16;
+struct SamplingPattern
+{
+    int numSamples;
+    FVector2D ssOffsets[maxNumSamples];
+};
+const SamplingPattern g_ssPatterns[] =
+{
+    {
+        1,
+        {
+            FVector2D(0, 0),
+        }
+    },
+    {
+        4,
+        {
+            FVector2D(0.125f, 0.625f),
+            FVector2D(0.375f, 0.125f),
+            FVector2D(0.625f, 0.875f),
+            FVector2D(0.875f, 0.375f),
+        }
+    },
+    {
+        16,
+        {
+            FVector2D(0.125f, 0.125f),
+            FVector2D(0.125f, 0.375f),
+            FVector2D(0.125f, 0.625f),
+            FVector2D(0.125f, 0.875f),
+            FVector2D(0.375f, 0.125f),
+            FVector2D(0.375f, 0.375f),
+            FVector2D(0.375f, 0.625f),
+            FVector2D(0.375f, 0.875f),
+            FVector2D(0.625f, 0.125f),
+            FVector2D(0.625f, 0.375f),
+            FVector2D(0.625f, 0.625f),
+            FVector2D(0.625f, 0.875f),
+            FVector2D(0.875f, 0.125f),
+            FVector2D(0.875f, 0.375f),
+            FVector2D(0.875f, 0.625f),
+            FVector2D(0.875f, 0.875f),
+        }
+    },
+
+};
+
+void USceneCapturer::InitCaptureComponent( USceneCaptureComponent2D* CaptureComponent, float HFov, float VFov, EStereoscopicPass InStereoPass )
+{
+	CaptureComponent->SetVisibility( true );
+	CaptureComponent->SetHiddenInGame( false );
+
+    CaptureComponent->CaptureStereoPass = InStereoPass;
+    CaptureComponent->FOVAngle = FMath::Max( HFov, VFov );
+    CaptureComponent->bCaptureEveryFrame = false;
+    CaptureComponent->CaptureSource = ESceneCaptureSource::SCS_FinalColorLDR;
+
+    CaptureComponent->TextureTarget = NewObject<UTextureRenderTarget2D>(this);
+    //TODO: ikrimae: Not sure why the render target needs to be float to avoid banding. Seems like captures to this RT and then applies PP
+    //               on top of it which causes degredation.
+    CaptureComponent->TextureTarget->InitCustomFormat(CaptureWidth, CaptureHeight, PF_A16B16G16R16, false);
+	CaptureComponent->TextureTarget->ClearColor = FLinearColor::Red;
+
+	CaptureComponent->RegisterComponentWithWorld( GWorld );
+
+	// UE4 cannot serialize an array of subobject pointers, so add these objects to the root
+	CaptureComponent->AddToRoot();
+}
+
+USceneCapturer::USceneCapturer(FVTableHelper& Helper)
+    : Super(Helper)
+    , ImageWrapperModule(FModuleManager::LoadModuleChecked<IImageWrapperModule>(FName("ImageWrapper")))
+    , bIsTicking(false)
+    , CapturePlayerController(NULL)
+    , CaptureGameMode(NULL)
+    , hAngIncrement(FStereoPanoramaManager::HorizontalAngularIncrement->GetFloat())
+    , vAngIncrement(FStereoPanoramaManager::VerticalAngularIncrement->GetFloat())
+    , eyeSeparation(FStereoPanoramaManager::EyeSeparation->GetFloat())
+    , NumberOfHorizontalSteps((int32)(360.0f / hAngIncrement))
+    , NumberOfVerticalSteps((int32)(180.0f / vAngIncrement) + 1) /* Need an extra b/c we only grab half of the top & bottom slices */
+    , SphericalAtlasWidth(FStereoPanoramaManager::StepCaptureWidth->GetInt())
+    , SphericalAtlasHeight(SphericalAtlasWidth / 2)
+    , bForceAlpha(FStereoPanoramaManager::ForceAlpha->GetInt() != 0)
+    , bEnableBilerp(FStereoPanoramaManager::EnableBilerp->GetInt() != 0)
+    , SSMethod(FMath::Clamp<int32>(FStereoPanoramaManager::SuperSamplingMethod->GetInt(), 0, ARRAY_COUNT(g_ssPatterns)))
+    , bOverrideInitialYaw(FStereoPanoramaManager::ShouldOverrideInitialYaw->GetInt() != 0)
+    , ForcedInitialYaw(FRotator::ClampAxis(FStereoPanoramaManager::ForcedInitialYaw->GetFloat()))
+    , OutputDir(FStereoPanoramaManager::OutputDir->GetString().IsEmpty() ? FPaths::GameSavedDir() / TEXT("StereoPanorama") : FStereoPanoramaManager::OutputDir->GetString())
+    , dbgDisableOffsetRotation(FStereoPanoramaManager::FadeStereoToZeroAtSides->GetInt() != 0)
+{}
+
+USceneCapturer::USceneCapturer()
+	: ImageWrapperModule( FModuleManager::LoadModuleChecked<IImageWrapperModule>( FName( "ImageWrapper" ) ) )
+	, bIsTicking( false )
+	, CapturePlayerController( NULL )
+	, CaptureGameMode( NULL )
+    , hAngIncrement( FStereoPanoramaManager::HorizontalAngularIncrement->GetFloat() )
+    , vAngIncrement( FStereoPanoramaManager::VerticalAngularIncrement->GetFloat() )
+    , eyeSeparation( FStereoPanoramaManager::EyeSeparation->GetFloat() )
+    , NumberOfHorizontalSteps( ( int32 )( 360.0f / hAngIncrement ) )
+    , NumberOfVerticalSteps( ( int32 )( 180.0f / vAngIncrement ) + 1 ) /* Need an extra b/c we only grab half of the top & bottom slices */
+    , SphericalAtlasWidth( FStereoPanoramaManager::StepCaptureWidth->GetInt() )
+    , SphericalAtlasHeight( SphericalAtlasWidth / 2)
+    , bForceAlpha( FStereoPanoramaManager::ForceAlpha->GetInt() != 0 )
+    , bEnableBilerp( FStereoPanoramaManager::EnableBilerp->GetInt() != 0 )
+    , SSMethod( FMath::Clamp<int32>(FStereoPanoramaManager::SuperSamplingMethod->GetInt(), 0, ARRAY_COUNT(g_ssPatterns)) )
+    , bOverrideInitialYaw( FStereoPanoramaManager::ShouldOverrideInitialYaw->GetInt() != 0 )
+    , ForcedInitialYaw( FRotator::ClampAxis(FStereoPanoramaManager::ForcedInitialYaw->GetFloat()) )
+    , OutputDir( FStereoPanoramaManager::OutputDir->GetString().IsEmpty() ? FPaths::GameSavedDir() / TEXT("StereoPanorama") : FStereoPanoramaManager::OutputDir->GetString() )
+    , dbgDisableOffsetRotation( FStereoPanoramaManager::FadeStereoToZeroAtSides->GetInt() != 0 )
+{
+    //NOTE: ikrimae: Keeping the old sampling mechanism just until we're sure the new way is always better
+    dbgMatchCaptureSliceFovToAtlasSliceFov = false;
+
+    float captureHFov = 0, captureVFov = 0;
+
+    if (dbgMatchCaptureSliceFovToAtlasSliceFov)
+    {
+        //Slicing Technique 1: Match Capture Slice StripWidth to match the pixel dimensions of AtlasWidth/NumHorizSteps & s.t. stripwidth/stripheight fovs match hAngIncr & vAngIncr
+        //                     Legacy technique but allows setting the strip width to match atlas slice width
+        //                     Pretty wasteful and will break if CaptureHFov & hangIncr/vAngIncr diverge greatly b/c resultant texture will exceed GPU bounds
+        //                     StripHeight is computed based on solving CpxV = CpxH * SpxV / SpxH
+        //                                                               CpxV = CV   * SpxV / SV
+        //                                                               captureVfov = 2 * atan( tan(captureHfov / 2) * (SpxV / SpxH) )
+        sliceHFov = hAngIncrement;
+        sliceVFov = vAngIncrement;
+
+        //TODO: ikrimae: Also do a quick test to see if there are issues with setting fov to something really small ( < 1 degree)
+        //               And it does. Current noted issues: screen space effects like SSAO, AA, SSR are all off
+        //                                                  local eyeadaptation also causes problems. Should probably turn off all PostProcess effects
+        //                                                  small fovs cause floating point errors in the sampling function (probably a bug b/c no thought put towards that)
+        captureHFov = FStereoPanoramaManager::CaptureHorizontalFOV->GetFloat();
+
+        ensure(captureHFov >= hAngIncrement);
+
+        //TODO: ikrimae: In hindsight, there's no reason that strip size should be this at all. Just select a square FOV larger than hAngIncr & vAngIncr
+        //               and then sample the resulting plane accordingly. Remember when updating to this to recheck the math in resample function. Might
+        //               have made assumptions about capture slice dimensions matching the sample strips
+        StripWidth = SphericalAtlasWidth / NumberOfHorizontalSteps;
+        //The scenecapture cube won't allow horizontal & vertical fov to not match the aspect ratio so we have to compute the right dimensions here for square pixels
+        StripHeight = StripWidth * FMath::Tan(FMath::DegreesToRadians(vAngIncrement / 2.0f)) / FMath::Tan(FMath::DegreesToRadians(hAngIncrement / 2.0f));
+
+        const FVector2D slicePlaneDim = FVector2D(
+            2.0f * FMath::Tan(FMath::DegreesToRadians(hAngIncrement) / 2.0f),
+            2.0f * FMath::Tan(FMath::DegreesToRadians(vAngIncrement) / 2.0f));
+
+        const float capturePlaneWidth = 2.0f * FMath::Tan(FMath::DegreesToRadians(captureHFov) / 2.0f);
+
+        //TODO: ikrimae: This is just to let the rest of the existing code work. Sampling rate of the slice can be whatever.
+        //      Ex: To match the highest sampling frequency of the spherical atlas, it should match the area of differential patch
+        //      at ray direction of pixel(0,1) in the atlas
+
+        //Need stripwidth/slicePlaneDim.X = capturewidth / capturePlaneDim.X
+        CaptureWidth = capturePlaneWidth * StripWidth / slicePlaneDim.X;
+        CaptureHeight = CaptureWidth * StripHeight / StripWidth;
+
+        captureVFov = FMath::RadiansToDegrees(2 * FMath::Atan(FMath::Tan(FMath::DegreesToRadians(captureHFov / 2.0f)) * CaptureHeight / CaptureWidth));
+
+        //float dbgCapturePlaneDimY = 2.0f * FMath::Tan(FMath::DegreesToRadians(captureVFov) / 2.0f);
+        //float dbgCaptureHeight = dbgCapturePlaneDimY * StripHeight / slicePlaneDim.Y;
+    }
+    else
+    {
+        //Slicing Technique 2: Each slice is a determined square FOV at a configured preset resolution.
+        //                     Strip Width/Strip Height is determined based on hAngIncrement & vAngIncrement
+        //                     Just make sure pixels/captureHFov >= pixels/hAngIncr && pixels/vAngIncr
+
+        captureVFov = captureHFov = FStereoPanoramaManager::CaptureHorizontalFOV->GetFloat();
+        sliceVFov   = sliceHFov   = captureHFov;
+
+        ensure(captureHFov >= FMath::Max(hAngIncrement, vAngIncrement));
+        
+        //TODO: ikrimae: Re-do for floating point accuracy
+        const FVector2D slicePlaneDim = FVector2D(
+            2.0f * FMath::Tan(FMath::DegreesToRadians(hAngIncrement) / 2.0f),
+            2.0f * FMath::Tan(FMath::DegreesToRadians(vAngIncrement) / 2.0f));
+
+        const FVector2D capturePlaneDim = FVector2D(
+            2.0f * FMath::Tan(FMath::DegreesToRadians(captureHFov) / 2.0f),
+            2.0f * FMath::Tan(FMath::DegreesToRadians(captureVFov) / 2.0f));
+
+        CaptureHeight = CaptureWidth = FStereoPanoramaManager::CaptureSlicePixelWidth->GetInt();
+
+        StripWidth  = CaptureWidth  * slicePlaneDim.X / capturePlaneDim.X;
+        StripHeight = CaptureHeight * slicePlaneDim.Y / capturePlaneDim.Y;
+
+        //TODO: ikrimae: Come back and check for the actual right sampling rate
+        check(StripWidth  >=  (SphericalAtlasWidth / NumberOfHorizontalSteps) && 
+              StripHeight >= (SphericalAtlasHeight / NumberOfVerticalSteps));
+        
+        //Ensure Width/Height is always even
+        StripWidth  += StripWidth & 1;
+        StripHeight += StripHeight & 1;
+
+    }
+
+    UnprojectedAtlasWidth  = NumberOfHorizontalSteps * StripWidth;
+    UnprojectedAtlasHeight = NumberOfVerticalSteps   * StripHeight;
+
+    //NOTE: ikrimae: Ensure that the main gameview is > CaptureWidth x CaptureHeight. Bug in UE4 that won't re-alloc scene render targets to the correct size
+    //               when the scenecapture component > current window render target. https://answers.unrealengine.com/questions/80531/scene-capture-2d-max-resolution.html
+    //TODO: ikrimae: Ensure that r.SceneRenderTargetResizeMethod=2
+    FSystemResolution::RequestResolutionChange(CaptureWidth, CaptureHeight, EWindowMode::Windowed);
+
+
+
+	for( int CaptureIndex = 0; CaptureIndex < FStereoPanoramaManager::ConcurrentCaptures->GetInt(); CaptureIndex++ )
+	{
+		FString LeftCounter = FString::Printf( TEXT( "LeftEyeCaptureComponent_%04d" ), CaptureIndex );
+		USceneCaptureComponent2D* LeftEyeCaptureComponent = CreateDefaultSubobject<USceneCaptureComponent2D>( *LeftCounter );
+		InitCaptureComponent( LeftEyeCaptureComponent, captureHFov, captureVFov, EStereoscopicPass::eSSP_LEFT_EYE );
+		LeftEyeCaptureComponents.Add( LeftEyeCaptureComponent );
+
+		FString RightCounter = FString::Printf( TEXT( "RightEyeCaptureComponent_%04d" ), CaptureIndex );
+		USceneCaptureComponent2D* RightEyeCaptureComponent = CreateDefaultSubobject<USceneCaptureComponent2D>( *RightCounter );
+		InitCaptureComponent( RightEyeCaptureComponent, captureHFov, captureVFov, EStereoscopicPass::eSSP_RIGHT_EYE );
+		RightEyeCaptureComponents.Add( RightEyeCaptureComponent );
+	}
+
+	CurrentStep = 0;
+	TotalSteps = 0;
+	FrameDescriptors = TEXT( "FrameNumber, GameClock, TimeTaken(s)" LINE_TERMINATOR );
+
+	CaptureStep = ECaptureStep::Reset;
+}
+
+void USceneCapturer::Reset()
+{
+	for( int CaptureIndex = 0; CaptureIndex < FStereoPanoramaManager::ConcurrentCaptures->GetInt(); CaptureIndex++ )
+	{
+		USceneCaptureComponent2D* LeftEyeCaptureComponent = LeftEyeCaptureComponents[CaptureIndex];
+		USceneCaptureComponent2D* RightEyeCaptureComponent = RightEyeCaptureComponents[CaptureIndex];
+
+		LeftEyeCaptureComponent->SetVisibility( false );
+		LeftEyeCaptureComponent->SetHiddenInGame( true );
+		
+		// UE4 cannot serialize an array of subobject pointers, so work around the GC problems
+		LeftEyeCaptureComponent->RemoveFromRoot();
+
+		RightEyeCaptureComponent->SetVisibility( false );
+		RightEyeCaptureComponent->SetHiddenInGame( true );
+		
+		// UE4 cannot serialize an array of subobject pointers, so work around the GC problems
+		RightEyeCaptureComponent->RemoveFromRoot();
+	}
+
+	UnprojectedLeftEyeAtlas.Empty();
+	UnprojectedRightEyeAtlas.Empty();
+}
+
+void USceneCapturer::SetPositionAndRotation( int32 CurrentHorizontalStep, int32 CurrentVerticalStep, int32 CaptureIndex )
+{
+	FRotator Rotation = StartRotation;
+	Rotation.Yaw += CurrentHorizontalStep * hAngIncrement;
+	Rotation.Pitch -= CurrentVerticalStep * vAngIncrement;
+
+    Rotation = Rotation.Clamp();
+
+	FVector Offset( 0.0f, eyeSeparation / 2.0f, 0.0f );
+    if (dbgDisableOffsetRotation)
+    {
+        //For rendering near field objects, we don't rotate the capture components around the stereo pivot, but instead
+        //around each capture component
+        const auto rotAngleOffset = FRotator::ClampAxis(Rotation.Yaw - StartRotation.Yaw);
+        float eyeSeparationDampeningFactor = 1.0f;
+        if (rotAngleOffset <= 90.0f)
+        {
+            eyeSeparationDampeningFactor = FMath::Lerp(1.0f, 0.0f, rotAngleOffset / 90.0f);
+        }
+        else if (rotAngleOffset <= 270.0f)
+        {
+            eyeSeparationDampeningFactor = 0.0f;
+        }
+        else
+        {
+            eyeSeparationDampeningFactor = FMath::Lerp(0.0f, 1.0f, (rotAngleOffset - 270.0f) / 90.0f);
+        }
+
+        Offset = StartRotation.RotateVector(Offset * eyeSeparationDampeningFactor);
+    }
+    else
+    {
+        Offset = Rotation.RotateVector(Offset);
+    }
+
+	LeftEyeCaptureComponents[CaptureIndex]->SetWorldLocationAndRotation( StartLocation - Offset, Rotation );
+    LeftEyeCaptureComponents[CaptureIndex]->UpdateContent();
+	RightEyeCaptureComponents[CaptureIndex]->SetWorldLocationAndRotation( StartLocation + Offset, Rotation );
+    RightEyeCaptureComponents[CaptureIndex]->UpdateContent();
+}
+
+void USceneCapturer::ValidateParameters()
+{
+	// Angular increment needs to be a factor of 360 to avoid seams i.e. 360 / angular increment needs to be a whole number
+	if( ( int32 )( NumberOfHorizontalSteps * hAngIncrement ) != 360 )
+	{
+		UE_LOG( LogStereoPanorama, Warning, TEXT( "Horizontal angular step (%g) is not a factor of 360! This will lead to a seam between the start and end points" ), hAngIncrement );
+	}
+
+	if( ( int32 )( (NumberOfVerticalSteps - 1) * vAngIncrement ) != 180 )
+	{
+		UE_LOG( LogStereoPanorama, Warning, TEXT( "Vertical angular step (%g) is not a factor of 180! This will lead to a seam between the start and end points" ), vAngIncrement );
+	}
+
+	TotalSteps = NumberOfHorizontalSteps * NumberOfVerticalSteps;
+    if( ( SphericalAtlasWidth & 1 ) != 0)
+    {
+        UE_LOG(LogStereoPanorama, Warning, TEXT("The Atlas Width (%d) must be even! Otherwise the Atlas height will not divide evenly."), SphericalAtlasWidth);
+    }
+
+
+	// The strip width needs to be an even number and a factor of the number of steps
+	if( ( StripWidth & 1 ) != 0 )
+	{
+		UE_LOG( LogStereoPanorama, Warning, TEXT( "Strip width (%d) needs to be even to avoid bad offsets" ), StripWidth );
+	}
+
+	if( StripWidth * NumberOfHorizontalSteps != SphericalAtlasWidth )
+	{
+		UE_LOG( LogStereoPanorama, Warning, TEXT( "The number of horizontal steps (%d) needs to be a factor of the atlas width (%d)" ), NumberOfHorizontalSteps, SphericalAtlasWidth );
+	}
+
+    if ((StripHeight & 1) != 0)
+    {
+        UE_LOG(LogStereoPanorama, Warning, TEXT("Strip height (%d) needs to be even to avoid bad offsets"), StripHeight);
+    }
+
+    if (StripHeight * (NumberOfVerticalSteps - 1) != SphericalAtlasHeight)
+	{
+        UE_LOG(LogStereoPanorama, Warning, TEXT("The number of vertical steps (%d) needs to be a factor of the atlas height (%d)"), NumberOfVerticalSteps, SphericalAtlasHeight);
+	}
+
+    //TODO: ikrimae: Validate capturewidth & captureheight. Need to be even
+
+	UE_LOG( LogStereoPanorama, Display, TEXT( "Stereo panoramic screenshot parameters" ) );
+	UE_LOG( LogStereoPanorama, Display, TEXT( " ... capture size: %d x %d" ), CaptureWidth, CaptureHeight );
+	UE_LOG( LogStereoPanorama, Display, TEXT( " ... spherical atlas size: %d x %d" ), SphericalAtlasWidth, SphericalAtlasHeight );
+    UE_LOG( LogStereoPanorama, Display, TEXT( " ... intermediate atlas size: %d x %d" ), UnprojectedAtlasWidth, UnprojectedAtlasHeight );
+	UE_LOG( LogStereoPanorama, Display, TEXT( " ... strip size: %d x %d" ), StripWidth, StripHeight );
+	UE_LOG( LogStereoPanorama, Display, TEXT( " ... horizontal steps: %d at %g degrees" ), NumberOfHorizontalSteps, hAngIncrement );
+	UE_LOG( LogStereoPanorama, Display, TEXT( " ... vertical steps: %d at %g degrees" ), NumberOfVerticalSteps, vAngIncrement );
+}
+
+void USceneCapturer::SetInitialState( int32 InStartFrame, int32 InEndFrame, FStereoCaptureDoneDelegate& InStereoCaptureDoneDelegate )
+{
+	if( bIsTicking )
+	{
+		UE_LOG( LogStereoPanorama, Warning, TEXT( "Already capturing a scene; concurrent captures are not allowed" ) );
+		return;
+	}
+
+	CapturePlayerController = UGameplayStatics::GetPlayerController( GWorld, 0 );
+	CaptureGameMode = UGameplayStatics::GetGameMode( GWorld );
+
+	if( CaptureGameMode == NULL || CapturePlayerController == NULL )
+	{
+		UE_LOG( LogStereoPanorama, Warning, TEXT( "Missing GameMode or PlayerController" ) );
+		return;
+	}
+
+	// Calculate the steps and validate they will produce good results
+	ValidateParameters();
+
+	// Setup starting criteria
+    StartFrame        = InStartFrame;
+    EndFrame          = InEndFrame;
+	CurrentFrameCount = 0;
+    CurrentStep       = 0;
+    CaptureStep       = ECaptureStep::Unpause;
+
+	Timestamp = FString::Printf( TEXT( "%s" ), *FDateTime::Now().ToString() );
+	
+	//SetStartPosition();
+
+	// Create storage for atlas textures
+    check( UnprojectedAtlasWidth * UnprojectedAtlasHeight <= MAX_int32 );
+	UnprojectedLeftEyeAtlas.AddUninitialized(  UnprojectedAtlasWidth * UnprojectedAtlasHeight );
+    UnprojectedRightEyeAtlas.AddUninitialized( UnprojectedAtlasWidth * UnprojectedAtlasHeight );
+
+	StartTime        = FDateTime::UtcNow();
+	OverallStartTime = StartTime;
+	bIsTicking       = true;
+
+    StereoCaptureDoneDelegate = InStereoCaptureDoneDelegate;
+}
+
+void USceneCapturer::CopyToUnprojAtlas( int32 CurrentHorizontalStep, int32 CurrentVerticalStep, TArray<FColor>& Atlas, TArray<FColor>& SurfaceData )
+{
+	int32 XOffset = StripWidth * CurrentHorizontalStep;
+    int32 YOffset = StripHeight * CurrentVerticalStep;
+
+	int32 StripSize = StripWidth * sizeof( FColor );
+    for (int32 Y = 0; Y < StripHeight; Y++)
+	{
+        void* Destination = &Atlas[( ( Y + YOffset ) * UnprojectedAtlasWidth ) + XOffset];
+		void* Source = &SurfaceData[StripWidth * Y];
+		FMemory::Memcpy( Destination, Source, StripSize );
+	}
+}
+
+TArray<FColor> USceneCapturer::SaveAtlas(FString Folder, const TArray<FColor>& SurfaceData)
+{
+	SCOPE_CYCLE_COUNTER( STAT_SPSavePNG );
+	
+    TArray<FColor> SphericalAtlas;
+    SphericalAtlas.AddZeroed(SphericalAtlasWidth * SphericalAtlasHeight);
+
+    const FVector2D slicePlaneDim = FVector2D(
+        2.0f * FMath::Tan(FMath::DegreesToRadians(sliceHFov) / 2.0f),
+        2.0f * FMath::Tan(FMath::DegreesToRadians(sliceVFov) / 2.0f));
+
+    //For each direction,
+    //    Find corresponding slice
+    //    Calculate intersection of slice plane
+    //    Calculate intersection UVs by projecting onto plane tangents
+    //    Supersample that UV coordinate from the unprojected atlas
+    {
+        SCOPE_CYCLE_COUNTER(STAT_SPSampleSpherical);
+        // Dump out how long the process took
+        const FDateTime SamplingStartTime = FDateTime::UtcNow();
+        UE_LOG(LogStereoPanorama, Log, TEXT("Sampling atlas..."));
+
+        for (int32 y = 0; y < SphericalAtlasHeight; y++)
+        {
+            for (int32 x = 0; x < SphericalAtlasWidth; x++)
+            {
+                FLinearColor samplePixelAccum = FLinearColor(0, 0, 0, 0);
+
+                //TODO: ikrimae: Seems that bilinear filtering sans supersampling is good enough. Supersampling sans bilerp seems best.
+                //               After more tests, come back to optimize by folding supersampling in and remove this outer sampling loop.
+                const auto& ssPattern = g_ssPatterns[SSMethod];
+
+                for (int32 i = 0; i < ssPattern.numSamples; i++)
+                {
+                    const float sampleU = ((float)x + ssPattern.ssOffsets[i].X) / SphericalAtlasWidth;
+                    const float sampleV = ((float)y + ssPattern.ssOffsets[i].Y) / SphericalAtlasHeight;
+
+                    const float sampleTheta = sampleU * 360.0f;
+                    const float samplePhi = sampleV * 180.0f;
+
+                    const FVector sampleDir = FVector(
+                        FMath::Sin(FMath::DegreesToRadians(samplePhi)) * FMath::Cos(FMath::DegreesToRadians(sampleTheta)),
+                        FMath::Sin(FMath::DegreesToRadians(samplePhi)) * FMath::Sin(FMath::DegreesToRadians(sampleTheta)),
+                        FMath::Cos(FMath::DegreesToRadians(samplePhi)));
+
+
+                    //TODO: ikrimae: ugh, ugly.
+                    const int32 sliceXIndex = FMath::TruncToInt(FRotator::ClampAxis(sampleTheta + hAngIncrement / 2.0f) / hAngIncrement);
+                    int32 sliceYIndex = 0;
+
+                    //Slice Selection = slice with max{sampleDir dot  sliceNormal }
+                    {
+                        float largestCosAngle = 0;
+                        for (int i = 0; i < NumberOfVerticalSteps; i++)
+                        {
+                            const FVector2D sliceCenterThetaPhi = FVector2D(
+                                hAngIncrement * sliceXIndex,
+                                vAngIncrement * i);
+
+                            //TODO: ikrimae: There has got to be a faster way. Rethink reparametrization later
+                            const FVector sliceDir = FVector(
+                                FMath::Sin(FMath::DegreesToRadians(sliceCenterThetaPhi.Y)) * FMath::Cos(FMath::DegreesToRadians(sliceCenterThetaPhi.X)),
+                                FMath::Sin(FMath::DegreesToRadians(sliceCenterThetaPhi.Y)) * FMath::Sin(FMath::DegreesToRadians(sliceCenterThetaPhi.X)),
+                                FMath::Cos(FMath::DegreesToRadians(sliceCenterThetaPhi.Y)));
+
+                            const float cosAngle = sampleDir | sliceDir;
+
+                            if (cosAngle > largestCosAngle)
+                            {
+                                largestCosAngle = cosAngle;
+                                sliceYIndex = i;
+                            }
+                        }
+                    }
+
+
+                    const FVector2D sliceCenterThetaPhi = FVector2D(
+                        hAngIncrement * sliceXIndex,
+                        vAngIncrement * sliceYIndex);
+
+                    //TODO: ikrimae: Reparameterize with an inverse mapping (e.g. project from slice pixels onto final u,v coordinates.
+                    //               Should make code simpler and faster b/c reduces to handful of sin/cos calcs per slice. 
+                    //               Supersampling will be more difficult though.
+
+                    const FVector sliceDir = FVector(
+                        FMath::Sin(FMath::DegreesToRadians(sliceCenterThetaPhi.Y)) * FMath::Cos(FMath::DegreesToRadians(sliceCenterThetaPhi.X)),
+                        FMath::Sin(FMath::DegreesToRadians(sliceCenterThetaPhi.Y)) * FMath::Sin(FMath::DegreesToRadians(sliceCenterThetaPhi.X)),
+                        FMath::Cos(FMath::DegreesToRadians(sliceCenterThetaPhi.Y)));
+
+                    const FPlane slicePlane = FPlane(sliceDir, -sliceDir);
+
+                    //Tangents from partial derivatives of sphere equation
+                    const FVector slicePlanePhiTangent = FVector(
+                        FMath::Cos(FMath::DegreesToRadians(sliceCenterThetaPhi.Y)) * FMath::Cos(FMath::DegreesToRadians(sliceCenterThetaPhi.X)),
+                        FMath::Cos(FMath::DegreesToRadians(sliceCenterThetaPhi.Y)) * FMath::Sin(FMath::DegreesToRadians(sliceCenterThetaPhi.X)),
+                        -FMath::Sin(FMath::DegreesToRadians(sliceCenterThetaPhi.Y))).GetSafeNormal();
+
+                    //Should be reconstructed to get around discontinuity of theta tangent at nodal points
+                    const FVector slicePlaneThetaTangent = (sliceDir ^ slicePlanePhiTangent).GetSafeNormal();
+                    //const FVector slicePlaneThetaTangent = FVector(
+                    //    -FMath::Sin(FMath::DegreesToRadians(sliceCenterThetaPhi.Y)) * FMath::Sin(FMath::DegreesToRadians(sliceCenterThetaPhi.X)),
+                    //    FMath::Sin(FMath::DegreesToRadians(sliceCenterThetaPhi.Y)) * FMath::Cos(FMath::DegreesToRadians(sliceCenterThetaPhi.X)),
+                    //    0).SafeNormal();
+
+                    check(!slicePlaneThetaTangent.IsZero() && !slicePlanePhiTangent.IsZero());
+
+                    const double t = (double)-slicePlane.W / (sampleDir | sliceDir);
+                    const FVector sliceIntersection = FVector(t * sampleDir.X, t * sampleDir.Y, t * sampleDir.Z);
+
+                    //Calculate scalar projection of sliceIntersection onto tangent vectors. a dot b / |b| = a dot b when tangent vectors are normalized
+                    //Then reparameterize to U,V of the sliceplane based on slice plane dimensions
+                    const float sliceU = (sliceIntersection | slicePlaneThetaTangent) / slicePlaneDim.X;
+                    const float sliceV = (sliceIntersection | slicePlanePhiTangent) / slicePlaneDim.Y;
+
+                    check(sliceU >= -(0.5f + KINDA_SMALL_NUMBER) &&
+                        sliceU <= (0.5f + KINDA_SMALL_NUMBER));
+
+                    check(sliceV >= -(0.5f + KINDA_SMALL_NUMBER) &&
+                        sliceV <= (0.5f + KINDA_SMALL_NUMBER));
+
+                    //TODO: ikrimae: Supersample/bilinear filter
+                    const int32 slicePixelX = FMath::TruncToInt(dbgMatchCaptureSliceFovToAtlasSliceFov ? sliceU * StripWidth : sliceU * CaptureWidth);
+                    const int32 slicePixelY = FMath::TruncToInt(dbgMatchCaptureSliceFovToAtlasSliceFov ? sliceV * StripHeight : sliceV * CaptureHeight);
+
+                    FLinearColor slicePixelSample;
+
+                    if (bEnableBilerp)
+                    {
+                        //TODO: ikrimae: Clean up later; too tired now
+                        const int32 sliceCenterPixelX = (sliceXIndex + 0.5f) * StripWidth;
+                        const int32 sliceCenterPixelY = (sliceYIndex + 0.5f) * StripHeight;
+
+                        const FIntPoint atlasSampleTL(sliceCenterPixelX + FMath::Clamp(slicePixelX    , -StripWidth/2, StripWidth/2), sliceCenterPixelY + FMath::Clamp(slicePixelY    , -StripHeight/2, StripHeight/2));
+                        const FIntPoint atlasSampleTR(sliceCenterPixelX + FMath::Clamp(slicePixelX + 1, -StripWidth/2, StripWidth/2), sliceCenterPixelY + FMath::Clamp(slicePixelY    , -StripHeight/2, StripHeight/2));
+                        const FIntPoint atlasSampleBL(sliceCenterPixelX + FMath::Clamp(slicePixelX    , -StripWidth/2, StripWidth/2), sliceCenterPixelY + FMath::Clamp(slicePixelY + 1, -StripHeight/2, StripHeight/2));
+                        const FIntPoint atlasSampleBR(sliceCenterPixelX + FMath::Clamp(slicePixelX + 1, -StripWidth/2, StripWidth/2), sliceCenterPixelY + FMath::Clamp(slicePixelY + 1, -StripHeight/2, StripHeight/2));
+
+                        const FColor pixelColorTL = SurfaceData[atlasSampleTL.Y * UnprojectedAtlasWidth + atlasSampleTL.X];
+                        const FColor pixelColorTR = SurfaceData[atlasSampleTR.Y * UnprojectedAtlasWidth + atlasSampleTR.X];
+                        const FColor pixelColorBL = SurfaceData[atlasSampleBL.Y * UnprojectedAtlasWidth + atlasSampleBL.X];
+                        const FColor pixelColorBR = SurfaceData[atlasSampleBR.Y * UnprojectedAtlasWidth + atlasSampleBR.X];
+
+                        const float fracX = FMath::Frac(dbgMatchCaptureSliceFovToAtlasSliceFov ? sliceU * StripWidth : sliceU * CaptureWidth);
+                        const float fracY = FMath::Frac(dbgMatchCaptureSliceFovToAtlasSliceFov ? sliceV * StripHeight : sliceV * CaptureHeight);
+
+                        //Reinterpret as linear (a.k.a dont apply srgb inversion)
+                        slicePixelSample = FMath::BiLerp(
+                            pixelColorTL.ReinterpretAsLinear(), pixelColorTR.ReinterpretAsLinear(),
+                            pixelColorBL.ReinterpretAsLinear(), pixelColorBR.ReinterpretAsLinear(),
+                            fracX, fracY);
+                    }
+                    else
+                    {
+                        const int32 sliceCenterPixelX = (sliceXIndex + 0.5f) * StripWidth;
+                        const int32 sliceCenterPixelY = (sliceYIndex + 0.5f) * StripHeight;
+
+                        const int32 atlasSampleX = sliceCenterPixelX + slicePixelX;
+                        const int32 atlasSampleY = sliceCenterPixelY + slicePixelY;
+
+
+                        slicePixelSample = SurfaceData[atlasSampleY * UnprojectedAtlasWidth + atlasSampleX].ReinterpretAsLinear();
+                    }
+
+                    samplePixelAccum += slicePixelSample;
+
+                    ////Output color map of projections
+                    //const FColor debugEquiColors[12] = {
+                    //    FColor(205, 180, 76),
+                    //    FColor(190, 88, 202),
+                    //    FColor(127, 185, 194),
+                    //    FColor(90, 54, 47),
+                    //    FColor(197, 88, 53),
+                    //    FColor(197, 75, 124),
+                    //    FColor(130, 208, 72),
+                    //    FColor(136, 211, 153),
+                    //    FColor(126, 130, 207),
+                    //    FColor(83, 107, 59),
+                    //    FColor(200, 160, 157),
+                    //    FColor(80, 66, 106)
+                    //};
+
+                    //samplePixelAccum = ssPattern.numSamples * debugEquiColors[sliceYIndex * 4 + sliceXIndex];
+                }
+
+                SphericalAtlas[y * SphericalAtlasWidth + x] = (samplePixelAccum / ssPattern.numSamples).Quantize();
+
+                // Force alpha value
+                if (bForceAlpha)
+                {
+                    SphericalAtlas[y * SphericalAtlasWidth + x].A = 255;
+                }
+            }
+        }
+
+        //Blit the first column into the last column to make the stereo image seamless at theta=360
+        for (int32 y = 0; y < SphericalAtlasHeight; y++)
+        {
+            SphericalAtlas[y * SphericalAtlasWidth + (SphericalAtlasWidth - 1)] = SphericalAtlas[y * SphericalAtlasWidth + 0];
+        }
+
+        const FTimespan SamplingDuration = FDateTime::UtcNow() - SamplingStartTime;
+        UE_LOG(LogStereoPanorama, Log, TEXT("...done! Duration: %g seconds"), SamplingDuration.GetTotalSeconds());
+    }
+	
+	// Generate name
+	FString FrameString = FString::Printf( TEXT( "%s_%05d.png" ), *Folder, CurrentFrameCount );
+    FString AtlasName =  OutputDir / Timestamp / FrameString;
+    
+	UE_LOG( LogStereoPanorama, Log, TEXT( "Writing atlas: %s" ), *AtlasName );
+
+	// Write out PNG
+    //TODO: ikrimae: Use threads to write out the images for performance
+	IImageWrapperPtr ImageWrapper = ImageWrapperModule.CreateImageWrapper( EImageFormat::PNG );
+    ImageWrapper->SetRaw(SphericalAtlas.GetData(), SphericalAtlas.GetAllocatedSize(), SphericalAtlasWidth, SphericalAtlasHeight, ERGBFormat::BGRA, 8);
+	const TArray<uint8>& PNGData = ImageWrapper->GetCompressed(100);
+	FFileHelper::SaveArrayToFile( PNGData, *AtlasName );
+
+    if (FStereoPanoramaManager::GenerateDebugImages->GetInt() != 0)
+    {
+        FString FrameStringUnprojected = FString::Printf(TEXT("%s_%05d_Unprojected.png"), *Folder, CurrentFrameCount);
+        FString AtlasNameUnprojected = OutputDir / Timestamp / FrameStringUnprojected;
+
+        ImageWrapper->SetRaw(SurfaceData.GetData(), SurfaceData.GetAllocatedSize(), UnprojectedAtlasWidth, UnprojectedAtlasHeight, ERGBFormat::BGRA, 8);
+        const TArray<uint8>& PNGDataUnprojected = ImageWrapper->GetCompressed(100);
+        FFileHelper::SaveArrayToFile(PNGData, *AtlasNameUnprojected);
+    }
+	ImageWrapper.Reset();
+
+	UE_LOG( LogStereoPanorama, Log, TEXT( " ... done!" ), *AtlasName );
+
+    return SphericalAtlas;
+}
+
+void USceneCapturer::CaptureComponent( int32 CurrentHorizontalStep, int32 CurrentVerticalStep, FString Folder, USceneCaptureComponent2D* CaptureComponent, TArray<FColor>& Atlas )
+{
+	TArray<FColor> SurfaceData;
+
+	{
+		SCOPE_CYCLE_COUNTER( STAT_SPReadStrip );
+		FTextureRenderTargetResource* RenderTarget = CaptureComponent->TextureTarget->GameThread_GetRenderTargetResource();
+
+		//TODO: ikrimae: Might need to validate that this divides evenly. Might not matter
+		int32 CenterX = CaptureWidth / 2;
+		int32 CenterY = CaptureHeight / 2;
+
+		SurfaceData.AddUninitialized( StripWidth * StripHeight );
+
+		// Read pixels
+		FIntRect Area( CenterX - ( StripWidth / 2 ), CenterY - ( StripHeight / 2 ), CenterX + ( StripWidth / 2 ), CenterY + ( StripHeight / 2) );
+        auto readSurfaceDataFlags = FReadSurfaceDataFlags();
+        readSurfaceDataFlags.SetLinearToGamma(false);
+		RenderTarget->ReadPixelsPtr( SurfaceData.GetData(), readSurfaceDataFlags, Area );
+	}
+
+	// Copy off strip to atlas texture
+	CopyToUnprojAtlas( CurrentHorizontalStep, CurrentVerticalStep, Atlas, SurfaceData );
+
+	if( FStereoPanoramaManager::GenerateDebugImages->GetInt() != 0 )
+	{
+		SCOPE_CYCLE_COUNTER( STAT_SPSavePNG );
+
+		// Generate name
+		FString TickString = FString::Printf( TEXT( "_%05d_%04d_%04d" ), CurrentFrameCount, CurrentHorizontalStep, CurrentVerticalStep );
+		FString CaptureName = OutputDir / Timestamp / Folder / TickString + TEXT( ".png" );
+		UE_LOG( LogStereoPanorama, Log, TEXT( "Writing snapshot: %s" ), *CaptureName );
+
+		// Write out PNG
+        if (FStereoPanoramaManager::GenerateDebugImages->GetInt() == 2)
+        {
+            //Read Whole Capture Buffer
+		    IImageWrapperPtr ImageWrapper = ImageWrapperModule.CreateImageWrapper( EImageFormat::PNG );
+
+            TArray<FColor> SurfaceDataWhole;
+            SurfaceDataWhole.AddUninitialized(CaptureWidth * CaptureHeight);
+            // Read pixels
+            FTextureRenderTargetResource* RenderTarget = CaptureComponent->TextureTarget->GameThread_GetRenderTargetResource();
+            RenderTarget->ReadPixelsPtr(SurfaceDataWhole.GetData(), FReadSurfaceDataFlags());
+
+            // Force alpha value
+            if (bForceAlpha)
+            {
+                for (FColor& Color : SurfaceDataWhole)
+                {
+                    Color.A = 255;
+                }
+            }
+
+            ImageWrapper->SetRaw(SurfaceDataWhole.GetData(), SurfaceDataWhole.GetAllocatedSize(), CaptureWidth, CaptureHeight, ERGBFormat::BGRA, 8);
+            const TArray<uint8>& PNGData = ImageWrapper->GetCompressed(100);
+
+            FFileHelper::SaveArrayToFile(PNGData, *CaptureName);
+            ImageWrapper.Reset();
+        }
+        else
+        {
+            if (bForceAlpha)
+            {
+                for (FColor& Color : SurfaceData)
+                {
+                    Color.A = 255;
+                }
+            }
+
+            IImageWrapperPtr ImageWrapper = ImageWrapperModule.CreateImageWrapper(EImageFormat::PNG);
+            ImageWrapper->SetRaw(SurfaceData.GetData(), SurfaceData.GetAllocatedSize(), StripWidth, StripHeight, ERGBFormat::BGRA, 8);
+		    const TArray<uint8>& PNGData = ImageWrapper->GetCompressed(100);
+
+		    FFileHelper::SaveArrayToFile( PNGData, *CaptureName );
+		    ImageWrapper.Reset();
+	    }
+    }
+}
+
+//TODO: ikrimae: This is a cluster fuck. Come back and actually work out the timings. Trickery b/c SceneCaptureCubes Tick at the end of the frame so we're effectively queuing up the next
+//               step (pause, unpause, setposition) for the next frame. FlushRenderingCommands() added haphazardly to test but didn't want to remove them so close to delivery. 
+//               Think through when we actually need to flush and document.
+void USceneCapturer::Tick( float DeltaTime )
+{
+	if( !bIsTicking )
+	{
+		return;
+	}
+
+    if ( CurrentFrameCount < StartFrame )
+    {
+        //Skip until we're at the frame we want to render
+        CurrentFrameCount++;
+        CaptureStep = ECaptureStep::Pause;
+    }
+	else if( CurrentStep < TotalSteps )
+	{
+        if (CaptureStep == ECaptureStep::Unpause)
+        {
+            FlushRenderingCommands();
+            CaptureGameMode->ClearPause();
+            //GPauseRenderingRealtimeClock = false;
+            CaptureStep = ECaptureStep::Pause;
+            FlushRenderingCommands();
+        }
+        else if (CaptureStep == ECaptureStep::Pause)
+        {
+            FlushRenderingCommands();
+            CaptureGameMode->SetPause(CapturePlayerController);
+            //GPauseRenderingRealtimeClock = true;
+            CaptureStep = ECaptureStep::SetStartPosition;
+            FlushRenderingCommands();
+        }
+        else if (CaptureStep == ECaptureStep::SetStartPosition)
+        {
+            //SetStartPosition();
+            ENQUEUE_UNIQUE_RENDER_COMMAND(
+                SceneCapturer_HeartbeatTickTickables,
+            {
+                TickRenderingTickables();
+            });
+
+            FlushRenderingCommands();
+            
+            FRotator Rotation;
+            CapturePlayerController->GetPlayerViewPoint(StartLocation, Rotation);
+            
+            Rotation.Roll = 0.0f;
+            Rotation.Yaw = (bOverrideInitialYaw) ? ForcedInitialYaw : Rotation.Yaw;
+            Rotation.Pitch = 90.0f;
+            StartRotation = Rotation;
+            CaptureStep = ECaptureStep::SetPosition;
+            FlushRenderingCommands();
+        }
+        else if (CaptureStep == ECaptureStep::SetPosition)
+        {
+            FlushRenderingCommands();
+            for (int32 CaptureIndex = 0; CaptureIndex < FStereoPanoramaManager::ConcurrentCaptures->GetInt(); CaptureIndex++)
+            {
+                int32 CurrentHorizontalStep;
+                int32 CurrentVerticalStep;
+                if (GetComponentSteps(CurrentStep + CaptureIndex, CurrentHorizontalStep, CurrentVerticalStep))
+                {
+                    SetPositionAndRotation(CurrentHorizontalStep, CurrentVerticalStep, CaptureIndex);
+                }
+            }
+
+            CaptureStep = ECaptureStep::Read;
+            FlushRenderingCommands();
+        }
+        else if (CaptureStep == ECaptureStep::Read)
+        {
+            FlushRenderingCommands();
+            for (int32 CaptureIndex = 0; CaptureIndex < FStereoPanoramaManager::ConcurrentCaptures->GetInt(); CaptureIndex++)
+            {
+                int32 CurrentHorizontalStep;
+                int32 CurrentVerticalStep;
+                if (GetComponentSteps(CurrentStep, CurrentHorizontalStep, CurrentVerticalStep))
+                {
+                    CaptureComponent(CurrentHorizontalStep, CurrentVerticalStep, TEXT("Left"), LeftEyeCaptureComponents[CaptureIndex], UnprojectedLeftEyeAtlas);
+                    CaptureComponent(CurrentHorizontalStep, CurrentVerticalStep, TEXT("Right"), RightEyeCaptureComponents[CaptureIndex], UnprojectedRightEyeAtlas);
+
+                    CurrentStep++;
+                }
+            }
+
+            CaptureStep = ECaptureStep::SetPosition;
+            FlushRenderingCommands();
+        }
+        else
+        {
+            //ECaptureStep::Reset:
+		}
+	}
+	else
+	{
+		TArray<FColor> SphericalLeftEyeAtlas  = SaveAtlas( TEXT( "Left" ), UnprojectedLeftEyeAtlas );
+        TArray<FColor> SphericalRightEyeAtlas = SaveAtlas(TEXT("Right"), UnprojectedRightEyeAtlas);
+		
+		// Dump out how long the process took
+		FDateTime EndTime = FDateTime::UtcNow();
+		FTimespan Duration = EndTime - StartTime;
+		UE_LOG( LogStereoPanorama, Log, TEXT( "Duration: %g seconds for frame %d" ), Duration.GetTotalSeconds(), CurrentFrameCount );
+		StartTime = EndTime;
+
+        //NOTE: ikrimae: Since we can't synchronously finish a stereocapture, we have to notify the caller with a function pointer
+        //Not sure this is the cleanest way but good enough for now
+        StereoCaptureDoneDelegate.ExecuteIfBound(SphericalLeftEyeAtlas, SphericalRightEyeAtlas);
+
+		// Construct log of saved atlases in csv format
+		FrameDescriptors += FString::Printf( TEXT( "%d, %g, %g" LINE_TERMINATOR ), CurrentFrameCount, FApp::GetCurrentTime() - FApp::GetLastTime(), Duration.GetTotalSeconds() );
+
+		CurrentFrameCount++;
+		if( CurrentFrameCount <= EndFrame )
+		{
+            CurrentStep = 0;
+			CaptureStep = ECaptureStep::Unpause;
+		}
+		else
+		{
+			CaptureGameMode->ClearPause();
+            //GPauseRenderingRealtimeClock = false;
+
+			FTimespan OverallDuration = FDateTime::UtcNow() - OverallStartTime;
+
+            FrameDescriptors += FString::Printf(TEXT("Duration: %g minutes for frame range [%d,%d] "), OverallDuration.GetTotalMinutes(), StartFrame, EndFrame);;
+			UE_LOG( LogStereoPanorama, Log, TEXT("Duration: %g minutes for frame range [%d,%d] "), OverallDuration.GetTotalMinutes(), StartFrame, EndFrame );
+
+			FString FrameDescriptorName = OutputDir / Timestamp / TEXT( "Frames.txt" );
+			FFileHelper::SaveStringToFile( FrameDescriptors, *FrameDescriptorName, FFileHelper::EEncodingOptions::ForceUTF8 );
+
+			bIsTicking = false;
+			FStereoPanoramaModule::Get()->Cleanup();
+		}
+	}
+}
diff --git a/Engine/Plugins/Experimental/StereoPanorama/Source/StereoPanorama/Private/SceneCapturer.h b/Engine/Plugins/Experimental/StereoPanorama/Source/StereoPanorama/Private/SceneCapturer.h
new file mode 100644
--- /dev/null
+++ b/Engine/Plugins/Experimental/StereoPanorama/Source/StereoPanorama/Private/SceneCapturer.h
@@ -0,0 +1,154 @@
+// Copyright 2015 Kite & Lightning.  All rights reserved.
+
+#pragma once
+
+#include "SceneCapturer.generated.h"
+
+DECLARE_LOG_CATEGORY_EXTERN( LogStereoPanorama, Log, All );
+
+DECLARE_STATS_GROUP( TEXT( "SP" ), STATGROUP_SP, STATCAT_Advanced );
+
+DECLARE_CYCLE_STAT( TEXT( "SavePNG" ),         STAT_SPSavePNG,         STATGROUP_SP );
+DECLARE_CYCLE_STAT( TEXT( "SampleSpherical" ), STAT_SPSampleSpherical, STATGROUP_SP );
+DECLARE_CYCLE_STAT( TEXT( "ReadStrip" ),       STAT_SPReadStrip,       STATGROUP_SP );
+DECLARE_CYCLE_STAT( TEXT( "FillAlpha" ),       STAT_SPFillAlpha,       STATGROUP_SP );
+
+UENUM()
+namespace ECaptureStep
+{
+	enum Type
+	{
+		Reset,
+        SetStartPosition,
+		SetPosition,
+		Read,
+		Pause,
+		Unpause,
+	};
+}
+
+DECLARE_DELEGATE_TwoParams(FStereoCaptureDoneDelegate, const TArray<FColor>&, const TArray<FColor>&);
+
+UCLASS()
+class USceneCapturer 
+	: public UObject
+	, public FTickableGameObject
+{
+    GENERATED_BODY()
+
+public:
+    USceneCapturer();
+
+    //NOTE: ikrimae: Adding this ctor hack to fix the 4.8p2 problem with hot reload macros calling empty constructors
+    //               https://answers.unrealengine.com/questions/228042/48p2-compile-fails-on-class-with-non-default-const.html
+    USceneCapturer(FVTableHelper& Helper);
+    
+	virtual void Tick( float DeltaTime ) override;
+
+	virtual bool IsTickable() const 
+	{ 
+		return true; 
+	}
+
+	virtual bool IsTickableWhenPaused() const
+	{
+		return bIsTicking;
+	}
+
+	virtual TStatId GetStatId() const
+	{
+		RETURN_QUICK_DECLARE_CYCLE_STAT( USceneCapturer, STATGROUP_Tickables );
+	}
+
+	void InitCaptureComponent( USceneCaptureComponent2D* CaptureComponent, float HFov, float VFov, EStereoscopicPass InStereoPass );
+
+	void CaptureComponent( int32 CurrentHorizontalStep, int32 CurrentVerticalStep, FString Folder, USceneCaptureComponent2D* CaptureComponent, TArray<FColor>& Atlas );
+
+	void CopyToUnprojAtlas( int32 CurrentHorizontalStep, int32 CurrentVerticalStep, TArray<FColor>& Atlas, TArray<FColor>& SurfaceData );
+
+	TArray<FColor> SaveAtlas( FString Folder, const TArray<FColor>& SurfaceData );
+
+	void SetPositionAndRotation( int32 CurrentHorizontalStep, int32 CurrentVerticalStep, int32 CaptureIndex );
+
+	void ValidateParameters();
+
+	void SetInitialState( int32 InStartFrame, int32 InEndFrame, FStereoCaptureDoneDelegate& InStereoCaptureDoneDelegate );
+
+	void Reset();
+
+	IImageWrapperModule& ImageWrapperModule;
+
+	bool bIsTicking;
+	FDateTime OverallStartTime;
+	FDateTime StartTime;
+
+	FVector StartLocation;
+	FRotator StartRotation;
+	FString Timestamp;
+    int32 StartFrame;
+    int32 EndFrame;
+
+	ECaptureStep::Type CaptureStep;
+	int32 CurrentFrameCount;
+
+	int32 CaptureWidth;
+	int32 CaptureHeight;
+
+	int32 StripWidth;
+    int32 StripHeight;
+
+	class APlayerController* CapturePlayerController;
+	class AGameMode* CaptureGameMode;
+
+	TArray<USceneCaptureComponent2D*> LeftEyeCaptureComponents;
+	TArray<USceneCaptureComponent2D*> RightEyeCaptureComponents;
+
+	bool GetComponentSteps( int32 Step, int32& CurrentHorizontalStep, int32& CurrentVerticalStep )
+	{
+		if( Step < TotalSteps )
+		{
+			CurrentHorizontalStep = Step / NumberOfVerticalSteps;
+			CurrentVerticalStep = Step - ( CurrentHorizontalStep * NumberOfVerticalSteps );
+			return true;
+		}
+
+		return false;
+	}
+
+private:
+    const float hAngIncrement;
+    const float vAngIncrement;
+    const float eyeSeparation;
+    
+    float sliceHFov;
+    float sliceVFov;
+
+    const int32 NumberOfHorizontalSteps;
+    const int32 NumberOfVerticalSteps;
+
+    int32 UnprojectedAtlasWidth;
+    int32 UnprojectedAtlasHeight;
+
+    const int32 SphericalAtlasWidth;
+    const int32 SphericalAtlasHeight;
+
+    int32 CurrentStep;
+	int32 TotalSteps;
+
+    TArray<FColor> UnprojectedLeftEyeAtlas;
+    TArray<FColor> UnprojectedRightEyeAtlas;
+
+    const bool bForceAlpha;
+
+    const bool bEnableBilerp;
+    const int32 SSMethod;
+    const bool bOverrideInitialYaw;
+    const float ForcedInitialYaw;
+    const FString OutputDir;
+
+    bool dbgMatchCaptureSliceFovToAtlasSliceFov;
+    bool dbgDisableOffsetRotation;
+	FString FrameDescriptors;
+
+    FStereoCaptureDoneDelegate StereoCaptureDoneDelegate;
+};
\ No newline at end of file
diff --git a/Engine/Plugins/Experimental/StereoPanorama/Source/StereoPanorama/Private/StereoCapturePawn.cpp b/Engine/Plugins/Experimental/StereoPanorama/Source/StereoPanorama/Private/StereoCapturePawn.cpp
new file mode 100644
--- /dev/null
+++ b/Engine/Plugins/Experimental/StereoPanorama/Source/StereoPanorama/Private/StereoCapturePawn.cpp
@@ -0,0 +1,86 @@
+// Copyright 2015 Kite & Lightning.  All rights reserved.
+
+#include "StereoPanoramaPrivatePCH.h"
+
+
+void AStereoCapturePawn::UpdateStereoAtlas(UObject* WorldContextObject, struct FLatentActionInfo LatentInfo)
+{
+    FIntPoint atlasDimensions(FStereoPanoramaManager::StepCaptureWidth->GetInt(), FStereoPanoramaManager::StepCaptureWidth->GetInt() / 2);
+    if (!LeftEyeAtlas || 
+        !LeftEyeAtlas->IsValidLowLevel() ||
+        LeftEyeAtlas->GetSizeX() != atlasDimensions.X ||
+        LeftEyeAtlas->GetSizeY() != atlasDimensions.Y)
+    {
+        LeftEyeAtlas = UTexture2D::CreateTransient(atlasDimensions.X, atlasDimensions.Y, EPixelFormat::PF_B8G8R8A8);
+        LeftEyeAtlas->SRGB = false;
+        //LeftEyeAtlas->CompressionNone = true;
+        LeftEyeAtlas->Filter = TextureFilter::TF_Trilinear;
+    }
+    
+    if (!RightEyeAtlas ||
+        !RightEyeAtlas->IsValidLowLevel() ||
+        RightEyeAtlas->GetSizeX() != atlasDimensions.X ||
+        RightEyeAtlas->GetSizeY() != atlasDimensions.Y)
+    {
+        RightEyeAtlas = UTexture2D::CreateTransient(atlasDimensions.X, atlasDimensions.Y, EPixelFormat::PF_B8G8R8A8);
+        RightEyeAtlas->SRGB = false;
+        //RightEyeAtlas->CompressionNone = true;
+        RightEyeAtlas->Filter = TextureFilter::TF_Trilinear;
+    }
+    FStereoCaptureDoneDelegate CopyAtlasToTexDelegate;
+    //TODO: ikrimae: When updating to 4.6, use stateful lambda instead of creating a separate member function
+    CopyAtlasToTexDelegate.BindUObject(this, &AStereoCapturePawn::CopyAtlasDataToTextures);
+    FStereoPanoramaModule::Get()->PanoramicScreenshot(0, 0, CopyAtlasToTexDelegate);
+
+    //TODO: ikrimae: Not cool b/c we're passing memory ownership to the LatentActionManager but at the same time we need a reference to this action to update when it's done
+    //               Got to dig deeper into UE4 to see proper way
+    StereoCaptureDoneAction = new FStereoCaptureDoneAction(LatentInfo);
+
+    if (UWorld* World = GEngine->GetWorldFromContextObject(WorldContextObject))
+    {
+        FLatentActionManager& LatentActionManager = World->GetLatentActionManager();
+        if (LatentActionManager.FindExistingAction<FStereoCaptureDoneAction>(LatentInfo.CallbackTarget, LatentInfo.UUID) == NULL)
+        {
+            LatentActionManager.AddNewAction(LatentInfo.CallbackTarget, LatentInfo.UUID, StereoCaptureDoneAction);
+        }
+    }
+}
+
+void AStereoCapturePawn::CopyAtlasDataToTextures(const TArray<FColor>& InLeftEyeAtlasData, const TArray<FColor>& InRightEyeAtlasData)
+{
+    if (LeftEyeAtlas &&
+        LeftEyeAtlas->IsValidLowLevel() &&
+        RightEyeAtlas &&
+        RightEyeAtlas->IsValidLowLevel())
+    {
+        {
+            const int32 TextureDataSize = InLeftEyeAtlasData.Num() * InLeftEyeAtlasData.GetTypeSize();
+            check(TextureDataSize == LeftEyeAtlas->GetSizeX() * LeftEyeAtlas->GetSizeY() * sizeof(FColor) );
+
+            FTexture2DMipMap& Mip = LeftEyeAtlas->PlatformData->Mips[0];
+            void* Data = Mip.BulkData.Lock(LOCK_READ_WRITE);
+            FMemory::Memcpy(Data, InLeftEyeAtlasData.GetTypedData(), TextureDataSize);
+            Mip.BulkData.Unlock();
+            LeftEyeAtlas->UpdateResource();
+        }
+
+        {
+            const int32 TextureDataSize = InRightEyeAtlasData.Num() * InRightEyeAtlasData.GetTypeSize();
+            check(TextureDataSize == RightEyeAtlas->GetSizeX() * RightEyeAtlas->GetSizeY() * sizeof(FColor) );
+
+            FTexture2DMipMap& Mip = RightEyeAtlas->PlatformData->Mips[0];
+            void* Data = Mip.BulkData.Lock(LOCK_READ_WRITE);
+            FMemory::Memcpy(Data, InRightEyeAtlasData.GetTypedData(), TextureDataSize);
+
+            Mip.BulkData.Unlock();
+            RightEyeAtlas->UpdateResource();
+        }
+    }
+
+    if (StereoCaptureDoneAction)
+    {
+        StereoCaptureDoneAction->IsStereoCaptureDone = true;
+    }
+}
+
+
diff --git a/Engine/Plugins/Experimental/StereoPanorama/Source/StereoPanorama/Private/StereoPanorama.cpp b/Engine/Plugins/Experimental/StereoPanorama/Source/StereoPanorama/Private/StereoPanorama.cpp
new file mode 100644
--- /dev/null
+++ b/Engine/Plugins/Experimental/StereoPanorama/Source/StereoPanorama/Private/StereoPanorama.cpp
@@ -0,0 +1,28 @@
+// Copyright 2015 Kite & Lightning.  All rights reserved.
+
+#include "StereoPanoramaPrivatePCH.h"
+
+TSharedPtr<FStereoPanoramaManager> StereoPanoramaManager;
+
+/** IModuleInterface - initialize the module */
+void FStereoPanoramaModule::StartupModule()
+{
+	StereoPanoramaManager = MakeShareable( new FStereoPanoramaManager() );
+}
+
+/** IModuleInterface - shutdown the module */
+void FStereoPanoramaModule::ShutdownModule()
+{
+	if( StereoPanoramaManager.IsValid() )
+	{
+		StereoPanoramaManager.Reset();
+	}
+}
+
+TSharedPtr<FStereoPanoramaManager> FStereoPanoramaModule::Get()
+{
+	check( StereoPanoramaManager.IsValid() );
+	return StereoPanoramaManager;
+}
+
+IMPLEMENT_MODULE( FStereoPanoramaModule, StereoPanorama )
diff --git a/Engine/Plugins/Experimental/StereoPanorama/Source/StereoPanorama/Private/StereoPanoramaManager.cpp b/Engine/Plugins/Experimental/StereoPanorama/Source/StereoPanorama/Private/StereoPanoramaManager.cpp
new file mode 100644
--- /dev/null
+++ b/Engine/Plugins/Experimental/StereoPanorama/Source/StereoPanorama/Private/StereoPanoramaManager.cpp
@@ -0,0 +1,135 @@
+// Copyright 2015 Kite & Lightning.  All rights reserved.
+
+#include "StereoPanoramaPrivatePCH.h"
+
+
+//Slice Controls
+IConsoleVariable* FStereoPanoramaManager::HorizontalAngularIncrement = IConsoleManager::Get().RegisterConsoleVariable( TEXT( "SP.HorizontalAngularIncrement" ), 1.0f, TEXT( "The number of degrees per horizontal step. Must be a factor of 360." ), ECVF_Default );
+IConsoleVariable* FStereoPanoramaManager::VerticalAngularIncrement   = IConsoleManager::Get().RegisterConsoleVariable( TEXT( "SP.VerticalAngularIncrement" ), 90.0f, TEXT( "The number of degrees per vertical step. Must be a factor of 180." ), ECVF_Default );
+IConsoleVariable* FStereoPanoramaManager::CaptureHorizontalFOV       = IConsoleManager::Get().RegisterConsoleVariable( TEXT( "SP.CaptureHorizontalFOV"), 90.0f, TEXT("Horizontal FOV for scene capture component. Must be larger than SP.HorizontalAngularIncrement"), ECVF_Default );
+
+//Atlas Controls
+IConsoleVariable* FStereoPanoramaManager::StepCaptureWidth           = IConsoleManager::Get().RegisterConsoleVariable( TEXT( "SP.StepCaptureWidth" ), 4096, TEXT( "The final spherical atlas width" ), ECVF_Default );
+IConsoleVariable* FStereoPanoramaManager::EyeSeparation              = IConsoleManager::Get().RegisterConsoleVariable( TEXT( "SP.EyeSeparation" ), 6.4f, TEXT( "The separation of the stereo cameras" ), ECVF_Default );
+IConsoleVariable* FStereoPanoramaManager::ForceAlpha                 = IConsoleManager::Get().RegisterConsoleVariable( TEXT( "SP.ForceAlpha" ), false, TEXT( "Force the alpha value to completely opaque" ), ECVF_Default );
+
+//Sampling Controls
+IConsoleVariable* FStereoPanoramaManager::CaptureSlicePixelWidth     = IConsoleManager::Get().RegisterConsoleVariable( TEXT( "SP.CaptureSlicePixelWidth"), 2048, TEXT(" Capture Slice Pixel Dimension"), ECVF_Default);
+IConsoleVariable* FStereoPanoramaManager::EnableBilerp               = IConsoleManager::Get().RegisterConsoleVariable( TEXT( "SP.EnableBilerp"), true, TEXT("0 - No Filtering 1- Bilinear Filter slice samples"), ECVF_Default);
+IConsoleVariable* FStereoPanoramaManager::SuperSamplingMethod        = IConsoleManager::Get().RegisterConsoleVariable( TEXT( "SP.SuperSamplingMethod"), 1, TEXT(" 0 - No Supersampling, 1 - Rotated Grid SS"), ECVF_Default);
+
+//Debug Controls
+//IConsoleVariable* FStereoPanoramaManager::ConcurrentCaptures         = IConsoleManager::Get().RegisterConsoleVariable( TEXT( "SP.ConcurrentCaptures" ), 15, TEXT("The number of scene captures to capture at the same time"), ECVF_Default);
+IConsoleVariable* FStereoPanoramaManager::ConcurrentCaptures         = IConsoleManager::Get().RegisterConsoleVariable( TEXT("SP.ConcurrentCaptures"), 30, TEXT("The number of scene captures to capture at the same time"), ECVF_Default);
+IConsoleVariable* FStereoPanoramaManager::GenerateDebugImages        = IConsoleManager::Get().RegisterConsoleVariable( TEXT( "SP.GenerateDebugImages" ), 0, TEXT("0 - No Debug Images\n1 - Save out each strip as it is generated\n2 - Save each entire slice"), ECVF_Default);
+IConsoleVariable* FStereoPanoramaManager::OutputDir                  = IConsoleManager::Get().RegisterConsoleVariable( TEXT( "SP.OutputDir" ), TEXT("F:/NeuroSP"), TEXT("Output directory"), ECVF_Default);
+IConsoleVariable* FStereoPanoramaManager::ShouldOverrideInitialYaw   = IConsoleManager::Get().RegisterConsoleVariable( TEXT( "SP.ShouldOverrideInitialYaw" ), true, TEXT("Override Initial Camera Yaw. Set to true if you don't want to use PlayerController View Dir"), ECVF_Default);
+IConsoleVariable* FStereoPanoramaManager::ForcedInitialYaw           = IConsoleManager::Get().RegisterConsoleVariable( TEXT( "SP.ForcedInitialYaw" ), 90.0f, TEXT("Yaw value for initial Camera view direction. Set ShouldOverrideInitialYaw to true to use this value"), ECVF_Default);
+IConsoleVariable* FStereoPanoramaManager::FadeStereoToZeroAtSides    = IConsoleManager::Get().RegisterConsoleVariable( TEXT( "SP.FadeStereoToZeroAtSides"), false, TEXT("Fade stereo effect between left/right eye to zero at 90 degrees."), ECVF_Default);
+
+
+void FStereoPanoramaManager::PanoramicScreenshot(const TArray<FString>& Args)
+{
+    FStereoCaptureDoneDelegate emptyDelegate;
+    FStereoPanoramaManager::PanoramicScreenshot(0, 0, emptyDelegate);
+}
+
+void FStereoPanoramaManager::PanoramicScreenshot(const int32 InStartFrame, const int32 InEndFrame, FStereoCaptureDoneDelegate& InStereoCaptureDoneDelegate)
+{
+	// Construct a capturer that has stereo USceneCaptureComponent2D components
+    SceneCapturer = ConstructObject<USceneCapturer>(USceneCapturer::StaticClass());
+	SceneCapturer->AddToRoot();
+
+	// Rotation is ignored; always start from a yaw of zero
+	SceneCapturer->SetInitialState( InStartFrame, InEndFrame, InStereoCaptureDoneDelegate );
+}
+
+void FStereoPanoramaManager::Cleanup()
+{
+	if( SceneCapturer != NULL )
+	{
+		SceneCapturer->Reset();
+		SceneCapturer->RemoveFromRoot();
+
+		// Let GC handle the deletion
+		SceneCapturer = NULL;
+	}
+}
+
+void FStereoPanoramaManager::PanoramicMovie( const TArray<FString>& Args )
+{
+    int32 StartFrame = 0;
+    int32 EndFrame   = 0;
+
+	if( Args.Num() == 1 )
+	{
+        StartFrame = 0;
+        EndFrame   = FCString::Atoi( *Args[0] ) - 1;    //Frame Range is inclusive so -1 to find the last frame
+	}
+    else if (Args.Num() == 2)
+    {
+        StartFrame = FCString::Atoi( *Args[0] );
+        EndFrame   = FCString::Atoi( *Args[1] );
+        EndFrame   = FMath::Max(StartFrame, EndFrame);
+    }
+
+    FStereoCaptureDoneDelegate emptyDelegate;
+    FStereoPanoramaManager::PanoramicScreenshot(StartFrame, EndFrame, emptyDelegate);
+}
+
+void FStereoPanoramaManager::PanoramicQuality( const TArray<FString>& Args )
+{
+	if( Args.Contains( TEXT( "preview" ) ) )
+	{
+		UE_LOG( LogStereoPanorama, Display, TEXT( " ... setting 'preview' quality" ) );
+
+		FStereoPanoramaManager::HorizontalAngularIncrement->Set( TEXT( "5" ) );
+		FStereoPanoramaManager::VerticalAngularIncrement->Set( TEXT( "60" ) );
+		FStereoPanoramaManager::CaptureHorizontalFOV->Set( TEXT( "60" ) );
+		FStereoPanoramaManager::StepCaptureWidth->Set( TEXT( "720" ) );
+	}
+	else if( Args.Contains( TEXT( "average" ) ) )
+	{
+		UE_LOG( LogStereoPanorama, Display, TEXT( " ... setting 'average' quality" ) );
+		
+		FStereoPanoramaManager::HorizontalAngularIncrement->Set( TEXT( "2" ) );
+		FStereoPanoramaManager::VerticalAngularIncrement->Set( TEXT( "30" ) );
+		FStereoPanoramaManager::CaptureHorizontalFOV->Set( TEXT( "30" ) );
+		FStereoPanoramaManager::StepCaptureWidth->Set( TEXT( "1440" ) );
+	}
+	else if( Args.Contains( TEXT( "improved" ) ) )
+	{
+		UE_LOG( LogStereoPanorama, Display, TEXT( " ... setting 'improved' quality" ) );
+		
+		FStereoPanoramaManager::HorizontalAngularIncrement->Set( TEXT( "0.5" ) );
+		FStereoPanoramaManager::VerticalAngularIncrement->Set( TEXT( "22.5" ) );
+		FStereoPanoramaManager::CaptureHorizontalFOV->Set( TEXT( "22.5" ) );
+		FStereoPanoramaManager::StepCaptureWidth->Set( TEXT( "1440" ) );
+	}
+	else
+	{
+		UE_LOG( LogStereoPanorama, Warning, TEXT( "No quality setting found; options are 'preview | average | improved'" ) );
+	}
+}
+
+void FStereoPanoramaManager::PanoramicTogglePause(const TArray<FString>& Args)
+{
+    auto CapturePlayerController = UGameplayStatics::GetPlayerController(GWorld, 0);
+    auto CaptureGameMode = UGameplayStatics::GetGameMode(GWorld);
+
+    if (CaptureGameMode == NULL || CapturePlayerController == NULL)
+    {
+        UE_LOG(LogStereoPanorama, Warning, TEXT("Missing GameMode or PlayerController"));
+        return;
+    }
+
+    if (GWorld->IsPaused())
+    {
+        CaptureGameMode->ClearPause();
+    }
+    else
+    {
+        CaptureGameMode->SetPause(CapturePlayerController);
+    }
+
+}
\ No newline at end of file
diff --git a/Engine/Plugins/Experimental/StereoPanorama/Source/StereoPanorama/Private/StereoPanoramaManager.h b/Engine/Plugins/Experimental/StereoPanorama/Source/StereoPanorama/Private/StereoPanoramaManager.h
new file mode 100644
--- /dev/null
+++ b/Engine/Plugins/Experimental/StereoPanorama/Source/StereoPanorama/Private/StereoPanoramaManager.h
@@ -0,0 +1,62 @@
+// Copyright 2015 Kite & Lightning.  All rights reserved.
+
+#pragma once
+#include "SceneCapturer.h"
+
+class FStereoPanoramaManager
+{
+private:
+	FAutoConsoleCommand PanoramicScreenshotCommand;
+	FAutoConsoleCommand PanoramicMovieCommand;
+	FAutoConsoleCommand PanoramicQualityCommand;
+    FAutoConsoleCommand PanoramicPauseCommand;
+
+public:
+    void PanoramicScreenshot( const TArray<FString>& Args );
+	void PanoramicScreenshot( const int32 InStartFrame, const int32 InEndFrame, FStereoCaptureDoneDelegate& InStereoCaptureDoneDelegate);
+	void PanoramicMovie( const TArray<FString>& Args );
+	void PanoramicQuality( const TArray<FString>& Args );
+    void PanoramicTogglePause(const TArray<FString>& Args);
+
+	static IConsoleVariable* HorizontalAngularIncrement;
+	static IConsoleVariable* VerticalAngularIncrement;
+	static IConsoleVariable* StepCaptureWidth;
+	static IConsoleVariable* EyeSeparation;
+	static IConsoleVariable* ForceAlpha;
+	static IConsoleVariable* GenerateDebugImages;
+	static IConsoleVariable* ConcurrentCaptures;
+    static IConsoleVariable* CaptureHorizontalFOV;
+    static IConsoleVariable* CaptureSlicePixelWidth;
+    static IConsoleVariable* EnableBilerp;
+    static IConsoleVariable* SuperSamplingMethod;
+    static IConsoleVariable* OutputDir;
+    static IConsoleVariable* ShouldOverrideInitialYaw;
+    static IConsoleVariable* ForcedInitialYaw;
+    static IConsoleVariable* FadeStereoToZeroAtSides;
+
+	FStereoPanoramaManager()
+		: PanoramicScreenshotCommand(
+			TEXT( "SP.PanoramicScreenshot" ),
+			*NSLOCTEXT( "StereoPanorama", "CommandText_ScreenShot", "Takes a panoramic screenshot" ).ToString(),
+			FConsoleCommandWithArgsDelegate::CreateRaw( this, &FStereoPanoramaManager::PanoramicScreenshot ) )
+		, PanoramicMovieCommand(
+			TEXT( "SP.PanoramicMovie" ),
+			*NSLOCTEXT( "StereoPanorama", "CommandText_MovieCapture", "Takes a sequence of panoramic screenshots" ).ToString(),
+			FConsoleCommandWithArgsDelegate::CreateRaw( this, &FStereoPanoramaManager::PanoramicMovie ) )
+		, PanoramicQualityCommand(
+			TEXT( "SP.PanoramicQuality" ),
+			*NSLOCTEXT( "StereoPanorama", "CommandText_Quality", "Sets the quality of the panoramic screenshot to 'preview | average | improved'" ).ToString(),
+			FConsoleCommandWithArgsDelegate::CreateRaw( this, &FStereoPanoramaManager::PanoramicQuality ) )
+		, PanoramicPauseCommand(
+            TEXT( "SP.TogglePause" ),
+			*NSLOCTEXT( "StereoPanorama", "CommandText_PauseGame", "Toggles Pausing/Unpausing of the game through StereoPanorama Plugin" ).ToString(),
+			FConsoleCommandWithArgsDelegate::CreateRaw( this, &FStereoPanoramaManager::PanoramicTogglePause ) )
+		, SceneCapturer( NULL )
+	{
+	}
+
+	void Cleanup();
+
+	UPROPERTY()
+	class USceneCapturer* SceneCapturer;
+};
\ No newline at end of file
diff --git a/Engine/Plugins/Experimental/StereoPanorama/Source/StereoPanorama/Private/StereoPanoramaPrivatePCH.h b/Engine/Plugins/Experimental/StereoPanorama/Source/StereoPanorama/Private/StereoPanoramaPrivatePCH.h
new file mode 100644
--- /dev/null
+++ b/Engine/Plugins/Experimental/StereoPanorama/Source/StereoPanorama/Private/StereoPanoramaPrivatePCH.h
@@ -0,0 +1,11 @@
+// Copyright 2015 Kite & Lightning.  All rights reserved.
+
+#include "Engine.h"
+#include "ModuleManager.h"
+#include "ImageWrapper.h"
+#include "Tickable.h"
+
+#include "StereoPanoramaManager.h"
+#include "StereoPanorama.h"
+#include "SceneCapturer.h"
+#include "StereoCapturePawn.h"
diff --git a/Engine/Plugins/Experimental/StereoPanorama/Source/StereoPanorama/Private/StereoStaticMeshComponent.cpp b/Engine/Plugins/Experimental/StereoPanorama/Source/StereoPanorama/Private/StereoStaticMeshComponent.cpp
new file mode 100644
--- /dev/null
+++ b/Engine/Plugins/Experimental/StereoPanorama/Source/StereoPanorama/Private/StereoStaticMeshComponent.cpp
@@ -0,0 +1,69 @@
+// Copyright 2015 Kite & Lightning.  All rights reserved.
+
+#include "StereoPanoramaPrivatePCH.h"
+#include "StereoStaticMeshComponent.h"
+#include "StaticMeshResources.h"
+
+class FStereoStaticMeshSceneProxy : public FStaticMeshSceneProxy
+{
+    ESPStereoCameraLayer::Type EyeToRender;
+
+public:
+
+    FStereoStaticMeshSceneProxy(UStereoStaticMeshComponent* Component) :
+        FStaticMeshSceneProxy(Component)
+    {
+        EyeToRender = Component->EyeToRender;
+    }
+
+    virtual FPrimitiveViewRelevance GetViewRelevance(const FSceneView* View) override
+    {
+        FPrimitiveViewRelevance viewRelevance = FStaticMeshSceneProxy::GetViewRelevance(View);
+        bool bVisible = true;
+
+        switch (View->StereoPass)
+        {
+        case eSSP_RIGHT_EYE:
+            if (EyeToRender != ESPStereoCameraLayer::RightEye && EyeToRender != ESPStereoCameraLayer::BothEyes)
+            {
+                bVisible = false;
+            }
+            break;
+
+        case eSSP_LEFT_EYE:
+            if (EyeToRender != ESPStereoCameraLayer::LeftEye && EyeToRender != ESPStereoCameraLayer::BothEyes)
+            {
+                bVisible = false;
+            }
+            break;
+
+        case eSSP_FULL:
+        default:
+            //Draw both planes when in mono mode
+            break;
+        }
+
+        viewRelevance.bDrawRelevance &= bVisible;
+
+        return viewRelevance;
+
+    }
+
+
+};
+
+FPrimitiveSceneProxy* UStereoStaticMeshComponent::CreateSceneProxy()
+{
+    if (StaticMesh == NULL
+        || StaticMesh->RenderData == NULL
+        || StaticMesh->RenderData->LODResources.Num() == 0
+        || StaticMesh->RenderData->LODResources[0].VertexBuffer.GetNumVertices() == 0)
+    {
+        return NULL;
+    }
+
+    FPrimitiveSceneProxy* Proxy = ::new FStereoStaticMeshSceneProxy(this);
+    return Proxy;
+}
+
+
diff --git a/Engine/Plugins/Experimental/StereoPanorama/Source/StereoPanorama/Private/StereoStaticMeshComponent.h b/Engine/Plugins/Experimental/StereoPanorama/Source/StereoPanorama/Private/StereoStaticMeshComponent.h
new file mode 100644
--- /dev/null
+++ b/Engine/Plugins/Experimental/StereoPanorama/Source/StereoPanorama/Private/StereoStaticMeshComponent.h
@@ -0,0 +1,37 @@
+// Copyright 2015 Kite & Lightning.  All rights reserved.
+
+#pragma once
+
+#include "CoreUObject.h"
+#include "Engine.h"
+#include "Components/StaticMeshComponent.h"
+#include "StereoStaticMeshComponent.generated.h"
+
+UENUM(BlueprintType)       //"BlueprintType" is essential specifier
+namespace ESPStereoCameraLayer
+{
+    //256 entries max
+    enum Type
+    {
+        LeftEye,
+        RightEye,
+        BothEyes,
+    };
+}
+
+/**
+ * 
+ */
+UCLASS(ClassGroup = (Rendering, Common), hidecategories = (Object, Activation, "Components|Activation"), ShowCategories = (Mobility), editinlinenew, meta = (BlueprintSpawnableComponent))
+class UStereoStaticMeshComponent : public UStaticMeshComponent
+{
+    GENERATED_BODY()
+
+public:
+
+    UPROPERTY(EditAnywhere, BlueprintReadWrite, Category = Sprite)
+    TEnumAsByte<ESPStereoCameraLayer::Type> EyeToRender;
+
+    virtual FPrimitiveSceneProxy* CreateSceneProxy() override;
+	
+};
diff --git a/Engine/Plugins/Experimental/StereoPanorama/Source/StereoPanorama/Public/StereoCapturePawn.h b/Engine/Plugins/Experimental/StereoPanorama/Source/StereoPanorama/Public/StereoCapturePawn.h
new file mode 100644
--- /dev/null
+++ b/Engine/Plugins/Experimental/StereoPanorama/Source/StereoPanorama/Public/StereoCapturePawn.h
@@ -0,0 +1,64 @@
+// Fill out your copyright notice in the Description page of Project Settings.
+
+#pragma once
+
+#include "GameFramework/Pawn.h"
+#include "LatentActions.h"
+#include "StereoCapturePawn.generated.h"
+
+class FStereoCaptureDoneAction : public FPendingLatentAction
+{
+public:
+    FName ExecutionFunction;
+    int32 OutputLink;
+    FWeakObjectPtr CallbackTarget;
+    bool IsStereoCaptureDone;
+
+    FStereoCaptureDoneAction(const FLatentActionInfo& LatentInfo)
+        : ExecutionFunction(LatentInfo.ExecutionFunction)
+        , OutputLink(LatentInfo.Linkage)
+        , CallbackTarget(LatentInfo.CallbackTarget)
+        , IsStereoCaptureDone(false)
+    {
+    }
+
+    virtual void UpdateOperation(FLatentResponse& Response)
+    {
+        Response.FinishAndTriggerIf(IsStereoCaptureDone, ExecutionFunction, OutputLink, CallbackTarget);
+    }
+
+#if WITH_EDITOR
+    // Returns a human readable description of the latent operation's current state
+    virtual FString GetDescription() const override
+    {
+        return FString::Printf(*NSLOCTEXT("StereoCaptureDoneAction", "IsStereoCaptureDone", "Is Stereo Capture Done: %d").ToString(), IsStereoCaptureDone);
+    }
+#endif
+};
+
+/**
+ * 
+ */
+UCLASS(config = Game, Blueprintable, BlueprintType)
+class AStereoCapturePawn : public ADefaultPawn
+{
+    GENERATED_BODY()
+
+public:
+
+    UFUNCTION(BlueprintCallable, Category = "Pawn", meta = (Latent, WorldContext = "WorldContextObject", LatentInfo = "LatentInfo"))
+    void UpdateStereoAtlas(UObject* WorldContextObject, struct FLatentActionInfo LatentInfo);
+
+    UPROPERTY(VisibleAnywhere, BlueprintReadOnly, Category = Pawn)
+    UTexture2D* LeftEyeAtlas;
+
+    UPROPERTY(VisibleAnywhere, BlueprintReadOnly, Category = Pawn)
+    UTexture2D* RightEyeAtlas;
+
+
+    FStereoCaptureDoneAction* StereoCaptureDoneAction;
+
+    void CopyAtlasDataToTextures(const TArray<FColor>& LeftEyeAtlasData, const TArray<FColor>& RightEyeAtlasData);
+
+	
+};
diff --git a/Engine/Plugins/Experimental/StereoPanorama/Source/StereoPanorama/Public/StereoPanorama.h b/Engine/Plugins/Experimental/StereoPanorama/Source/StereoPanorama/Public/StereoPanorama.h
new file mode 100644
--- /dev/null
+++ b/Engine/Plugins/Experimental/StereoPanorama/Source/StereoPanorama/Public/StereoPanorama.h
@@ -0,0 +1,16 @@
+// Copyright 2015 Kite & Lightning.  All rights reserved.
+
+#pragma once
+
+class FStereoPanoramaModule : public IModuleInterface
+{
+private:
+	/** IModuleInterface - initialize the module */
+	virtual void StartupModule() override;
+
+	/** IModuleInterface - shutdown the module */
+	virtual void ShutdownModule() override;
+
+public:
+	static TSharedPtr<FStereoPanoramaManager> Get();
+};
\ No newline at end of file
diff --git a/Engine/Plugins/Experimental/StereoPanorama/Source/StereoPanorama/StereoPanorama.Build.cs b/Engine/Plugins/Experimental/StereoPanorama/Source/StereoPanorama/StereoPanorama.Build.cs
new file mode 100644
--- /dev/null
+++ b/Engine/Plugins/Experimental/StereoPanorama/Source/StereoPanorama/StereoPanorama.Build.cs
@@ -0,0 +1,30 @@
+// Copyright 2015 Kite & Lightning.  All rights reserved.
+
+namespace UnrealBuildTool.Rules
+{
+	public class StereoPanorama : ModuleRules
+	{
+		public StereoPanorama( TargetInfo Target )
+		{
+			PrivateIncludePaths.AddRange(
+				new string[] {
+					"Runtime/StereoPanorama/Private",
+				}
+			);
+
+			PrivateDependencyModuleNames.AddRange(
+				new string[] {
+					"Core",
+					"CoreUObject",
+					"Engine",
+					"ImageWrapper",
+					"InputCore",
+					"RenderCore",
+					"ShaderCore",
+					"RHI",
+					"Slate",
+				}
+			);
+		}
+	}
+}
\ No newline at end of file
diff --git a/Engine/Plugins/Experimental/StereoPanorama/StereoPanorama.uplugin b/Engine/Plugins/Experimental/StereoPanorama/StereoPanorama.uplugin
new file mode 100644
--- /dev/null
+++ b/Engine/Plugins/Experimental/StereoPanorama/StereoPanorama.uplugin
@@ -0,0 +1,22 @@
+{
+    "FileVersion" : 1,
+	
+	"FriendlyName" : "Stereo Panoramic Movie Capture",
+	"Version" : 1,
+	"VersionName" : "Alpha",
+	"CreatedBy" : "Kite & Lightning",
+	"CreatedByURL" : "http://kiteandlightning.la/",
+	"EngineVersion" : 2328018,
+	"Description" : "A plugin to capture a sequence of panoramic images",
+	"Category" : "Movie",
+    "CanContainContent" : "true",
+	
+    "Modules" :
+	[
+		{
+			"Name" : "StereoPanorama",
+			"Type" : "RuntimeNoCommandlet",
+			"WhitelistPlatforms" : [ "Win64" ]
+		}
+	]
+}
\ No newline at end of file
diff --git a/Engine/Source/Runtime/Engine/Classes/Components/SceneCaptureComponent.h b/Engine/Source/Runtime/Engine/Classes/Components/SceneCaptureComponent.h
--- a/Engine/Source/Runtime/Engine/Classes/Components/SceneCaptureComponent.h
+++ b/Engine/Source/Runtime/Engine/Classes/Components/SceneCaptureComponent.h
@@ -44,6 +44,8 @@ class USceneCaptureComponent : public USceneComponent
 	FEngineShowFlags ShowFlags;
 
 public:
+	/** Indicates which stereo pass this component is capturing for, if any */
+    EStereoscopicPass CaptureStereoPass;
 
 	/** Adds the component to our list of hidden components. */
 	UFUNCTION(BlueprintCallable, Category = "Rendering|SceneCapture")
diff --git a/Engine/Source/Runtime/Engine/Classes/Matinee/InterpData.h b/Engine/Source/Runtime/Engine/Classes/Matinee/InterpData.h
--- a/Engine/Source/Runtime/Engine/Classes/Matinee/InterpData.h
+++ b/Engine/Source/Runtime/Engine/Classes/Matinee/InterpData.h
@@ -16,7 +16,7 @@ class UInterpData : public UObject
 	GENERATED_UCLASS_BODY()
 
 	/** Duration of interpolation sequence - in seconds. */
-	UPROPERTY()
+    UPROPERTY(VisibleAnywhere, BlueprintReadOnly, Category = InterpData)
 	float InterpLength;
 
 	/** Position in Interp to move things to for path-building in editor. */
diff --git a/Engine/Source/Runtime/Engine/Classes/Matinee/MatineeActor.h b/Engine/Source/Runtime/Engine/Classes/Matinee/MatineeActor.h
--- a/Engine/Source/Runtime/Engine/Classes/Matinee/MatineeActor.h
+++ b/Engine/Source/Runtime/Engine/Classes/Matinee/MatineeActor.h
@@ -179,7 +179,7 @@ public:
 #endif // WITH_EDITORONLY_DATA
 
 	/** properties that may change on InterpAction that we need to notify clients about, since the object's properties will not be replicated */
-	UPROPERTY(replicated,transient)
+    UPROPERTY(VisibleAnywhere, BlueprintReadOnly, replicated, transient, Category = Play)
 	uint32 bIsPlaying:1;
 
 	UPROPERTY(replicated)
@@ -192,7 +192,7 @@ public:
 	UPROPERTY(replicated,transient)
 	uint32 bPendingStop:1;
 
-	UPROPERTY(replicated)
+    UPROPERTY(VisibleAnywhere, BlueprintReadOnly, replicated, Category = Play)
 	float InterpPosition;
 
 	/** This is just optimization flag to skip checking it again. If all is initialized, it will set this to be true **/
diff --git a/Engine/Source/Runtime/Engine/Private/Components/SceneCaptureComponent.cpp b/Engine/Source/Runtime/Engine/Private/Components/SceneCaptureComponent.cpp
--- a/Engine/Source/Runtime/Engine/Private/Components/SceneCaptureComponent.cpp
+++ b/Engine/Source/Runtime/Engine/Private/Components/SceneCaptureComponent.cpp
@@ -158,6 +158,8 @@ USceneCaptureComponent::USceneCaptureComponent(const FObjectInitializer& ObjectI
 	ShowFlags.MotionBlur = 0; // motion blur doesn't work correctly with scene captures.
 	ShowFlags.SeparateTranslucency = 0;
 	ShowFlags.HMDDistortion = 0;
+
+    CaptureStereoPass = EStereoscopicPass::eSSP_FULL;
 }
 
 void USceneCaptureComponent::PostLoad()
@@ -278,6 +280,7 @@ USceneCaptureComponent2D::USceneCaptureComponent2D(const FObjectInitializer& Obj
 	CaptureSource = SCS_SceneColorHDR;
 	// default to full blend weight..
 	PostProcessBlendWeight = 1.0f;
+	CaptureStereoPass = EStereoscopicPass::eSSP_FULL;
 }
 
 void USceneCaptureComponent2D::OnRegister()
diff --git a/Engine/Source/Runtime/RenderCore/Public/RenderingThread.h b/Engine/Source/Runtime/RenderCore/Public/RenderingThread.h
--- a/Engine/Source/Runtime/RenderCore/Public/RenderingThread.h
+++ b/Engine/Source/Runtime/RenderCore/Public/RenderingThread.h
@@ -83,6 +83,7 @@ extern RENDERCORE_API void FlushRenderingCommands();
 extern RENDERCORE_API void FlushPendingDeleteRHIResources_GameThread();
 extern RENDERCORE_API void FlushPendingDeleteRHIResources_RenderThread();
 
+extern RENDERCORE_API void TickRenderingTickables();
 
 ////////////////////////////////////
 // Render thread suspension
diff --git a/Engine/Source/Runtime/Renderer/Private/SceneCaptureRendering.cpp b/Engine/Source/Runtime/Renderer/Private/SceneCaptureRendering.cpp
--- a/Engine/Source/Runtime/Renderer/Private/SceneCaptureRendering.cpp
+++ b/Engine/Source/Runtime/Renderer/Private/SceneCaptureRendering.cpp
@@ -177,6 +177,7 @@ FSceneRenderer* FScene::CreateSceneRenderer( USceneCaptureComponent* SceneCaptur
 	ViewInitOptions.BackgroundColor = FLinearColor::Black;
 	ViewInitOptions.OverrideFarClippingPlaneDistance = MaxViewDistance;
 	ViewInitOptions.SceneViewStateInterface = SceneCaptureComponent->GetViewState();
+    ViewInitOptions.StereoPass = SceneCaptureComponent->CaptureStereoPass;
 
 	if (bCaptureSceneColour)
 	{
