Fix minor warnings.

diff --git a/Engine/Plugins/Experimental/StereoPanorama/Source/StereoPanorama/Private/StereoCapturePawn.cpp b/Engine/Plugins/Experimental/StereoPanorama/Source/StereoPanorama/Private/StereoCapturePawn.cpp
--- a/Engine/Plugins/Experimental/StereoPanorama/Source/StereoPanorama/Private/StereoCapturePawn.cpp
+++ b/Engine/Plugins/Experimental/StereoPanorama/Source/StereoPanorama/Private/StereoCapturePawn.cpp
@@ -59,7 +59,7 @@ void AStereoCapturePawn::CopyAtlasDataToTextures(const TArray<FColor>& InLeftEye
 
             FTexture2DMipMap& Mip = LeftEyeAtlas->PlatformData->Mips[0];
             void* Data = Mip.BulkData.Lock(LOCK_READ_WRITE);
-            FMemory::Memcpy(Data, InLeftEyeAtlasData.GetTypedData(), TextureDataSize);
+            FMemory::Memcpy(Data, InLeftEyeAtlasData.GetData(), TextureDataSize);
             Mip.BulkData.Unlock();
             LeftEyeAtlas->UpdateResource();
         }
@@ -70,7 +70,7 @@ void AStereoCapturePawn::CopyAtlasDataToTextures(const TArray<FColor>& InLeftEye
 
             FTexture2DMipMap& Mip = RightEyeAtlas->PlatformData->Mips[0];
             void* Data = Mip.BulkData.Lock(LOCK_READ_WRITE);
-            FMemory::Memcpy(Data, InRightEyeAtlasData.GetTypedData(), TextureDataSize);
+            FMemory::Memcpy(Data, InRightEyeAtlasData.GetData(), TextureDataSize);
 
             Mip.BulkData.Unlock();
             RightEyeAtlas->UpdateResource();
diff --git a/Engine/Plugins/Experimental/StereoPanorama/Source/StereoPanorama/Private/StereoPanoramaManager.cpp b/Engine/Plugins/Experimental/StereoPanorama/Source/StereoPanorama/Private/StereoPanoramaManager.cpp
--- a/Engine/Plugins/Experimental/StereoPanorama/Source/StereoPanorama/Private/StereoPanoramaManager.cpp
+++ b/Engine/Plugins/Experimental/StereoPanorama/Source/StereoPanorama/Private/StereoPanoramaManager.cpp
@@ -37,7 +37,7 @@ void FStereoPanoramaManager::PanoramicScreenshot(const TArray<FString>& Args)
 void FStereoPanoramaManager::PanoramicScreenshot(const int32 InStartFrame, const int32 InEndFrame, FStereoCaptureDoneDelegate& InStereoCaptureDoneDelegate)
 {
 	// Construct a capturer that has stereo USceneCaptureComponent2D components
-    SceneCapturer = ConstructObject<USceneCapturer>(USceneCapturer::StaticClass());
+    SceneCapturer = NewObject<USceneCapturer>(USceneCapturer::StaticClass());
 	SceneCapturer->AddToRoot();
 
 	// Rotation is ignored; always start from a yaw of zero
